- fix incorrect copyright header in this .rc file, missed by previous checking cl #15657948

diff --git a/Engine/Source/Programs/EpicWebHelper/Resources/Windows/EpicWebHelper.rc b/Engine/Source/Programs/EpicWebHelper/Resources/Windows/EpicWebHelper.rc
--- a/Engine/Source/Programs/EpicWebHelper/Resources/Windows/EpicWebHelper.rc
+++ b/Engine/Source/Programs/EpicWebHelper/Resources/Windows/EpicWebHelper.rc
@@ -1,4 +1,4 @@
-// Copyright 1998-2019 Epic Games, Inc. All Rights Reserved.
+// Copyright Epic Games, Inc. All Rights Reserved.
 
 #include <windows.h>
 #include "Runtime/Launch/Resources/Version.h"
