- Fix compilation errors in EpicWebApp on OSX due to relatively old OSX min version, by disabling automation tests.
- Some dev automation tests in Core require "aligned new" which is unsupported in the older OSX version this target targets.

diff --git a/Engine/Source/Programs/EpicWebHelper/EpicWebHelper.Target.cs b/Engine/Source/Programs/EpicWebHelper/EpicWebHelper.Target.cs
--- a/Engine/Source/Programs/EpicWebHelper/EpicWebHelper.Target.cs
+++ b/Engine/Source/Programs/EpicWebHelper/EpicWebHelper.Target.cs
@@ -54,6 +54,9 @@ public class EpicWebHelperTarget : TargetRules
 		// Epic Games Launcher needs to run on OS X 10.9, so CEFSubProcess needs this as well
 		bEnableOSX109Support = true;
 
+		// Disable automation tests to avoid compilation errors related to alignas, caused by setting bEnableOSX109Support to true above. Old OSX versions do not fully support this until C++17
+		bForceDisableAutomationTests = true;
+
 		// Already a manifest specified through resource file
 		WindowsPlatform.ManifestFile = null;
 	}
