changing EpicWebHelper.app into console application

diff --git a/Engine/Build/Rsync/RsyncEngine.txt b/Engine/Build/Rsync/RsyncEngine.txt
--- a/Engine/Build/Rsync/RsyncEngine.txt
+++ b/Engine/Build/Rsync/RsyncEngine.txt
@@ -52,8 +52,7 @@
 + /Binaries/Mac/
 + /Binaries/Mac/DsymExporter.app
 + /Binaries/Mac/DsymExporter.app/***
-+ /Binaries/Mac/EpicWebHelper.app/
-+ /Binaries/Mac/EpicWebHelper.app/***
++ /Binaries/Mac/EpicWebHelper
 protect /Binaries/Mac/**
 
 # Binaries/IOS
diff --git a/Engine/Source/Programs/EpicWebHelper/EpicWebHelper.Target.cs b/Engine/Source/Programs/EpicWebHelper/EpicWebHelper.Target.cs
--- a/Engine/Source/Programs/EpicWebHelper/EpicWebHelper.Target.cs
+++ b/Engine/Source/Programs/EpicWebHelper/EpicWebHelper.Target.cs
@@ -28,7 +28,7 @@ public class EpicWebHelperTarget : TargetRules
 			var HelperTypes = new List<string>() { "GPU", "Renderer", "Plugin"};
 			foreach( string HelperType in HelperTypes)
 			{
-	            PostBuildSteps.Add(string.Format("cd $(EngineDir)/Binaries/$(TargetPlatform); if [ ! -h \"EpicWebHelper ({0}).app\" ]; then ln -s \"EpicWebHelper.app\" \"EpicWebHelper ({1}).app\"; fi", HelperType, HelperType));
+	            PostBuildSteps.Add(string.Format("cd $(EngineDir)/Binaries/$(TargetPlatform); if [ ! -h \"EpicWebHelper ({0})\" ]; then ln -s \"EpicWebHelper\" \"EpicWebHelper ({1})\"; fi", HelperType, HelperType));
 			}
         }
 
@@ -45,8 +45,16 @@ public class EpicWebHelperTarget : TargetRules
 		// Force all shader formats to be built and included.
 		//bForceBuildShaderFormats = true;
 
-		// CEFSubProcess is  a Windows app (uses WinMain())
-		bIsBuildingConsoleApplication = false;
+        if (Target.Platform == UnrealTargetPlatform.Mac)
+        {
+            // UE-201611 changing EpicWebHelper into a console application on Mac
+            bIsBuildingConsoleApplication = true;
+        }
+        else
+        {
+            // CEFSubProcess is a Windows app (uses WinMain())
+            bIsBuildingConsoleApplication = false;
+        }
 
 		// Disable logging, as the sub processes are spawned often and logging will just slow them down
 		GlobalDefinitions.Add("ALLOW_LOG_FILE=0");
diff --git a/Engine/Source/Runtime/WebBrowser/Private/WebBrowserSingleton.cpp b/Engine/Source/Runtime/WebBrowser/Private/WebBrowserSingleton.cpp
--- a/Engine/Source/Runtime/WebBrowser/Private/WebBrowserSingleton.cpp
+++ b/Engine/Source/Runtime/WebBrowser/Private/WebBrowserSingleton.cpp
@@ -70,7 +70,7 @@ THIRD_PARTY_INCLUDES_END
 #	elif PLATFORM_MAC
 #		define CEF3_FRAMEWORK_DIR CEF3_BIN_DIR TEXT("/Mac/Chromium Embedded Framework.framework")
 #		define CEF3_RESOURCES_DIR CEF3_FRAMEWORK_DIR TEXT("/Resources")
-#		define CEF3_SUBPROCES_EXE TEXT("Binaries/Mac/EpicWebHelper.app/Contents/MacOS/EpicWebHelper")
+#		define CEF3_SUBPROCES_EXE TEXT("Binaries/Mac/EpicWebHelper")
 #	elif PLATFORM_LINUX // @todo Linux
 #		define CEF3_RESOURCES_DIR CEF3_BIN_DIR TEXT("/Linux/Resources")
 #		define CEF3_SUBPROCES_EXE TEXT("Binaries/Linux/EpicWebHelper")
diff --git a/Engine/Source/Runtime/WebBrowser/WebBrowser.Build.cs b/Engine/Source/Runtime/WebBrowser/WebBrowser.Build.cs
--- a/Engine/Source/Runtime/WebBrowser/WebBrowser.Build.cs
+++ b/Engine/Source/Runtime/WebBrowser/WebBrowser.Build.cs
@@ -76,15 +76,7 @@ public class WebBrowser : ModuleRules
 
 			if (Target.Type != TargetType.Server)
 			{
-				if (Target.Platform == UnrealTargetPlatform.Mac)
-				{
-					// Add contents of EpicWebHelper.app directory as runtime dependencies
-					foreach (string FilePath in Directory.EnumerateFiles(Target.RelativeEnginePath + "/Binaries/Mac/EpicWebHelper.app", "*", SearchOption.AllDirectories))
-					{
-						RuntimeDependencies.Add(FilePath);
-					}
-				}
-				else if (Target.Platform == UnrealTargetPlatform.Linux)
+				if (Target.Platform == UnrealTargetPlatform.Mac || Target.Platform == UnrealTargetPlatform.Linux)
 				{
 					RuntimeDependencies.Add("$(EngineDir)/Binaries/" + Target.Platform.ToString() + "/EpicWebHelper");
 				}
