Fix EpicWebHelper missing compatibility manifest for Windows 10/11

The references to custom application manifests were removed from EpicWebHelper.rc in 15635368, possibly by accident, so the comment in EpicWebHelperTarget about a manifest already being provided was no longer true. Switching the target file to use the default engine manifest that sets Windows 10/11 compatibility, which was the original intent. This may help CEF choose faster and more secure code paths above Windows 8.

diff --git a/Engine/Source/Programs/EpicWebHelper/EpicWebHelper.Target.cs b/Engine/Source/Programs/EpicWebHelper/EpicWebHelper.Target.cs
--- a/Engine/Source/Programs/EpicWebHelper/EpicWebHelper.Target.cs
+++ b/Engine/Source/Programs/EpicWebHelper/EpicWebHelper.Target.cs
@@ -62,8 +62,5 @@ public class EpicWebHelperTarget : TargetRules
 		// Enable OSX 10.11 support (CEF3 requires 10.11+ on macOS)
 		bEnableOSX109Support = true;
 		//ForceOSXMinVersion = "10.11";
-
-		// Already a manifest specified through resource file
-		WindowsPlatform.ManifestFile = null;
 	}
 }
diff --git a/Engine/Source/Programs/EpicWebHelper/Resources/Windows/Debug-EpicWebHelper-Win32.manifest b/Engine/Source/Programs/EpicWebHelper/Resources/Windows/Debug-EpicWebHelper-Win32.manifest
deleted file mode 100644
diff --git a/Engine/Source/Programs/EpicWebHelper/Resources/Windows/Debug-EpicWebHelper-Win64.manifest b/Engine/Source/Programs/EpicWebHelper/Resources/Windows/Debug-EpicWebHelper-Win64.manifest
deleted file mode 100644
diff --git a/Engine/Source/Programs/EpicWebHelper/Resources/Windows/EpicWebHelper-Win32.manifest b/Engine/Source/Programs/EpicWebHelper/Resources/Windows/EpicWebHelper-Win32.manifest
deleted file mode 100644
diff --git a/Engine/Source/Programs/EpicWebHelper/Resources/Windows/EpicWebHelper-Win64.manifest b/Engine/Source/Programs/EpicWebHelper/Resources/Windows/EpicWebHelper-Win64.manifest
deleted file mode 100644
